Fix ShowLoginUI interface change in GameCircle plugin
#ue4
#android

diff --git a/Engine/Plugins/Online/Android/OnlineSubsystemGameCircle/Source/OnlineSubsystemGameCircle/Private/OnlineIdentityInterfaceGameCircle.cpp b/Engine/Plugins/Online/Android/OnlineSubsystemGameCircle/Source/OnlineSubsystemGameCircle/Private/OnlineIdentityInterfaceGameCircle.cpp
--- a/Engine/Plugins/Online/Android/OnlineSubsystemGameCircle/Source/OnlineSubsystemGameCircle/Private/OnlineIdentityInterfaceGameCircle.cpp
+++ b/Engine/Plugins/Online/Android/OnlineSubsystemGameCircle/Source/OnlineSubsystemGameCircle/Private/OnlineIdentityInterfaceGameCircle.cpp
@@ -30,7 +30,7 @@ TArray<TSharedPtr<FUserOnlineAccount> > FOnlineIdentityGameCircle::GetAllUserAcc
 
 bool FOnlineIdentityGameCircle::Login(int32 LocalUserNum, const FOnlineAccountCredentials& AccountCredentials) 
 {
-	return MainSubsystem->GetExternalUIInterface()->ShowLoginUI(LocalUserNum, false);
+	return MainSubsystem->GetExternalUIInterface()->ShowLoginUI(LocalUserNum, false, false);
 }
 
 
