Various fixes for CIS.

diff --git a/Engine/Plugins/Online/Android/OnlineSubsystemGameCircle/OnlineSubsystemGameCircle.uplugin b/Engine/Plugins/Online/Android/OnlineSubsystemGameCircle/OnlineSubsystemGameCircle.uplugin
--- a/Engine/Plugins/Online/Android/OnlineSubsystemGameCircle/OnlineSubsystemGameCircle.uplugin
+++ b/Engine/Plugins/Online/Android/OnlineSubsystemGameCircle/OnlineSubsystemGameCircle.uplugin
@@ -38,5 +38,11 @@
 			"Name": "OnlineSubsystemGameCircle",
 			"LoadingPolicy": "Always"
 		}
-	]	
+	],
+	"Plugins": [
+		{
+			"Name": "OnlineSubsystem",
+			"Enabled": true
+		}
+	]
 }
\ No newline at end of file
diff --git a/Engine/Plugins/Online/Android/OnlineSubsystemGooglePlay/OnlineSubsystemGooglePlay.uplugin b/Engine/Plugins/Online/Android/OnlineSubsystemGooglePlay/OnlineSubsystemGooglePlay.uplugin
--- a/Engine/Plugins/Online/Android/OnlineSubsystemGooglePlay/OnlineSubsystemGooglePlay.uplugin
+++ b/Engine/Plugins/Online/Android/OnlineSubsystemGooglePlay/OnlineSubsystemGooglePlay.uplugin
@@ -31,6 +31,10 @@
 		{
 			"Name": "OnlineSubsystem",
 			"Enabled": true
+		},
+		{
+			"Name": "AndroidPermission",
+			"Enabled": true
 		}
 	]
 }
\ No newline at end of file
diff --git a/Engine/Plugins/Runtime/LocationServicesAndroidImpl/LocationServicesAndroidImpl.uplugin b/Engine/Plugins/Runtime/LocationServicesAndroidImpl/LocationServicesAndroidImpl.uplugin
--- a/Engine/Plugins/Runtime/LocationServicesAndroidImpl/LocationServicesAndroidImpl.uplugin
+++ b/Engine/Plugins/Runtime/LocationServicesAndroidImpl/LocationServicesAndroidImpl.uplugin
@@ -30,6 +30,10 @@
 		}
 	],
 	"Plugins": [
+		{
+			"Name": "LocationServicesBPLibrary",
+			"Enabled": true
+		},
 		{
 			"Name": "AndroidPermission",
 			"Enabled": true
diff --git a/Engine/Plugins/Runtime/LocationServicesIOSImpl/LocationServicesIOSImpl.uplugin b/Engine/Plugins/Runtime/LocationServicesIOSImpl/LocationServicesIOSImpl.uplugin
--- a/Engine/Plugins/Runtime/LocationServicesIOSImpl/LocationServicesIOSImpl.uplugin
+++ b/Engine/Plugins/Runtime/LocationServicesIOSImpl/LocationServicesIOSImpl.uplugin
@@ -28,5 +28,11 @@
 				"IOS"
 			]
 		}
+	],
+	"Plugins": [
+		{
+			"Name": "LocationServicesBPLibrary",
+			"Enabled": true
+		}
 	]
 }
\ No newline at end of file
diff --git a/Engine/Source/Programs/CrashReporter/CrashReportCommon/SimpleWebRequest.cs b/Engine/Source/Programs/CrashReporter/CrashReportCommon/SimpleWebRequest.cs
--- a/Engine/Source/Programs/CrashReporter/CrashReportCommon/SimpleWebRequest.cs
+++ b/Engine/Source/Programs/CrashReporter/CrashReportCommon/SimpleWebRequest.cs
@@ -6,7 +6,7 @@ using System.IO;
 using System.Net;
 using System.Text;
 
-namespace Tools.DotNETCommon.SimpleWebRequest
+namespace Tools.CrashReporter.CrashReportCommon
 {
 	/// <summary>
 	/// A class to make simple Xml web requests to return a string.
