Fix non-unity builds for all target platforms.

diff --git a/Engine/Plugins/Online/Android/OnlineSubsystemGameCircle/Source/OnlineSubsystemGameCircle/Private/OnlineFriendGameCircle.cpp b/Engine/Plugins/Online/Android/OnlineSubsystemGameCircle/Source/OnlineSubsystemGameCircle/Private/OnlineFriendGameCircle.cpp
--- a/Engine/Plugins/Online/Android/OnlineSubsystemGameCircle/Source/OnlineSubsystemGameCircle/Private/OnlineFriendGameCircle.cpp
+++ b/Engine/Plugins/Online/Android/OnlineSubsystemGameCircle/Source/OnlineSubsystemGameCircle/Private/OnlineFriendGameCircle.cpp
@@ -2,6 +2,8 @@
 
 #include "OnlineFriendGameCircle.h"
 #include "Interfaces/OnlinePresenceInterface.h"
+#include "Interfaces/OnlineFriendsInterface.h"
+#include "Logging/LogMacros.h"
 #include "OnlineSubsystem.h"
 
 FOnlineFriendGameCircle::FOnlineFriendGameCircle(TSharedPtr<const FUniqueNetId> InUniqueId, const FString& InPlayerAlias, const FString& InAvatarURL)
diff --git a/Engine/Plugins/Online/OnlineSubsystemFacebook/Source/Private/IOS/OnlineIdentityFacebook.cpp b/Engine/Plugins/Online/OnlineSubsystemFacebook/Source/Private/IOS/OnlineIdentityFacebook.cpp
--- a/Engine/Plugins/Online/OnlineSubsystemFacebook/Source/Private/IOS/OnlineIdentityFacebook.cpp
+++ b/Engine/Plugins/Online/OnlineSubsystemFacebook/Source/Private/IOS/OnlineIdentityFacebook.cpp
@@ -6,8 +6,11 @@
 #include "Interfaces/OnlineSharingInterface.h"
 #include "Interfaces/OnlineExternalUIInterface.h"
 
+THIRD_PARTY_INCLUDES_START
 #import <FBSDKCoreKit/FBSDKCoreKit.h>
 #import <FBSDKLoginKit/FBSDKLoginKit.h>
+THIRD_PARTY_INCLUDES_END
+
 #import "FacebookHelper.h"
 
 // Other UE4 includes
diff --git a/Engine/Plugins/Online/OnlineSubsystemFacebook/Source/Private/IOS/OnlineSharingFacebook.cpp b/Engine/Plugins/Online/OnlineSubsystemFacebook/Source/Private/IOS/OnlineSharingFacebook.cpp
--- a/Engine/Plugins/Online/OnlineSubsystemFacebook/Source/Private/IOS/OnlineSharingFacebook.cpp
+++ b/Engine/Plugins/Online/OnlineSubsystemFacebook/Source/Private/IOS/OnlineSharingFacebook.cpp
@@ -8,9 +8,11 @@
 #include "IOS/IOSAsyncTask.h"
 #include "ImageCore.h"
 
+THIRD_PARTY_INCLUDES_START
 #import <FBSDKCoreKit/FBSDKCoreKit.h>
 #import <FBSDKShareKit/FBSDKShareKit.h>
 #import <FBSDKLoginKit/FBSDKLoginKit.h>
+THIRD_PARTY_INCLUDES_END
 
 FOnlineSharingFacebook::FOnlineSharingFacebook(FOnlineSubsystemFacebook* InSubsystem)
 	: FOnlineSharingFacebookCommon(InSubsystem)
diff --git a/Engine/Plugins/Runtime/Oculus/OculusAudio/Source/OculusAudio/Private/OculusAudioLegacy.h b/Engine/Plugins/Runtime/Oculus/OculusAudio/Source/OculusAudio/Private/OculusAudioLegacy.h
--- a/Engine/Plugins/Runtime/Oculus/OculusAudio/Source/OculusAudio/Private/OculusAudioLegacy.h
+++ b/Engine/Plugins/Runtime/Oculus/OculusAudio/Source/OculusAudio/Private/OculusAudioLegacy.h
@@ -1,6 +1,11 @@
 // Copyright 1998-2019 Epic Games, Inc. All Rights Reserved.
-#if PLATFORM_WINDOWS
+
 #pragma once
+
+#include "CoreTypes.h"
+
+#if PLATFORM_WINDOWS
+
 #include "OVR_Audio.h"
 #include "XAudio2Device.h"
 #include "AudioEffect.h"
diff --git a/Engine/Source/Runtime/Launch/Private/IOS/TcpConsoleListener.h b/Engine/Source/Runtime/Launch/Private/IOS/TcpConsoleListener.h
--- a/Engine/Source/Runtime/Launch/Private/IOS/TcpConsoleListener.h
+++ b/Engine/Source/Runtime/Launch/Private/IOS/TcpConsoleListener.h
@@ -3,6 +3,7 @@
 #pragma once
 
 #include "CoreMinimal.h"
+#include "HAL/Runnable.h"
 #include "Interfaces/IPv4/IPv4Endpoint.h"
 
 class FSocket;
diff --git a/Engine/Source/Runtime/Online/SSL/Public/Interfaces/ISslCertificateManager.h b/Engine/Source/Runtime/Online/SSL/Public/Interfaces/ISslCertificateManager.h
--- a/Engine/Source/Runtime/Online/SSL/Public/Interfaces/ISslCertificateManager.h
+++ b/Engine/Source/Runtime/Online/SSL/Public/Interfaces/ISslCertificateManager.h
@@ -2,6 +2,9 @@
 
 #pragma once
 
+#include "CoreTypes.h"
+#include "CoreFwd.h"
+
 struct ssl_ctx_st;
 typedef struct ssl_ctx_st SSL_CTX;
 
