wrapped a few more internal calls to deprecated OSS interfaces

diff --git a/Engine/Plugins/Online/Android/OnlineSubsystemGameCircle/Source/OnlineSubsystemGameCircle/Private/OnlineStoreInterfaceGameCircle.cpp b/Engine/Plugins/Online/Android/OnlineSubsystemGameCircle/Source/OnlineSubsystemGameCircle/Private/OnlineStoreInterfaceGameCircle.cpp
--- a/Engine/Plugins/Online/Android/OnlineSubsystemGameCircle/Source/OnlineSubsystemGameCircle/Private/OnlineStoreInterfaceGameCircle.cpp
+++ b/Engine/Plugins/Online/Android/OnlineSubsystemGameCircle/Source/OnlineSubsystemGameCircle/Private/OnlineStoreInterfaceGameCircle.cpp
@@ -113,10 +113,12 @@ JNI_METHOD void Java_com_epicgames_ue4_AmazonStoreHelper_nativeQueryComplete(JNI
 			if (IOnlineSubsystem* const OnlineSub = IOnlineSubsystem::Get())
 			{
 				// call store implementation to process query results.
+				PRAGMA_DISABLE_DEPRECATION_WARNINGS
 				if (FOnlineStoreGameCircle* StoreInterface = (FOnlineStoreGameCircle*)OnlineSub->GetStoreInterface().Get())
 				{
 					StoreInterface->ProcessQueryAvailablePurchasesResults(Result, ProvidedProductInformation);
 				}
+				PRAGMA_ENABLE_DEPRECATION_WARNINGS
 			}
 			FPlatformMisc::LowLevelOutputDebugStringf(TEXT("In-App Purchase query was completed  %s\n"), Result == EInAppPurchaseState::Success ? TEXT("successfully") : TEXT("unsuccessfully"));
 		}),
@@ -201,10 +203,12 @@ JNI_METHOD void Java_com_epicgames_ue4_AmazonStoreHelper_nativePurchaseComplete(
 			if (IOnlineSubsystem* const OnlineSub = IOnlineSubsystem::Get())
 			{
 				// call store implementation to process query results.
+				PRAGMA_DISABLE_DEPRECATION_WARNINGS
 				if (FOnlineStoreGameCircle* StoreInterface = (FOnlineStoreGameCircle*)OnlineSub->GetStoreInterface().Get())
 				{
 					StoreInterface->ProcessPurchaseResult(Result, ProductId, ReceiptData, Signature);
 				}
+				PRAGMA_ENABLE_DEPRECATION_WARNINGS
 			}
 		}),
 		GET_STATID(STAT_FSimpleDelegateGraphTask_ProcessIapResult), 
@@ -300,6 +304,7 @@ JNI_METHOD void Java_com_epicgames_ue4_AmazonStoreHelper_nativeRestorePurchasesC
 		{
 			FPlatformMisc::LowLevelOutputDebugStringf(TEXT("Sending result back to OnlineSubsystem.\n"));
 			// call store implementation to process query results.
+			PRAGMA_DISABLE_DEPRECATION_WARNINGS
 			if (FOnlineStoreGameCircle* StoreInterface = (FOnlineStoreGameCircle*)OnlineSub->GetStoreInterface().Get())
 			{
 				if (StoreInterface->CachedPurchaseRestoreObject.IsValid())
@@ -309,6 +314,7 @@ JNI_METHOD void Java_com_epicgames_ue4_AmazonStoreHelper_nativeRestorePurchasesC
 				}
 				StoreInterface->TriggerOnInAppPurchaseRestoreCompleteDelegates(bSuccess ? EInAppPurchaseState::Restored : Result);
 			}
+			PRAGMA_ENABLE_DEPRECATION_WARNINGS
 		}
 		}),
 		GET_STATID(STAT_FSimpleDelegateGraphTask_RestorePurchases),
