#online,store
Deferring OnlinePurchaseInterface deprecation to 4.26 instead of 4.25

diff --git a/Engine/Plugins/Online/Android/OnlineSubsystemGameCircle/Source/OnlineSubsystemGameCircle/Private/OnlineStoreInterfaceGameCircle.cpp b/Engine/Plugins/Online/Android/OnlineSubsystemGameCircle/Source/OnlineSubsystemGameCircle/Private/OnlineStoreInterfaceGameCircle.cpp
--- a/Engine/Plugins/Online/Android/OnlineSubsystemGameCircle/Source/OnlineSubsystemGameCircle/Private/OnlineStoreInterfaceGameCircle.cpp
+++ b/Engine/Plugins/Online/Android/OnlineSubsystemGameCircle/Source/OnlineSubsystemGameCircle/Private/OnlineStoreInterfaceGameCircle.cpp
@@ -113,12 +113,10 @@ JNI_METHOD void Java_com_epicgames_ue4_AmazonStoreHelper_nativeQueryComplete(JNI
 			if (IOnlineSubsystem* const OnlineSub = IOnlineSubsystem::Get())
 			{
 				// call store implementation to process query results.
-				PRAGMA_DISABLE_DEPRECATION_WARNINGS
 				if (FOnlineStoreGameCircle* StoreInterface = (FOnlineStoreGameCircle*)OnlineSub->GetStoreInterface().Get())
 				{
 					StoreInterface->ProcessQueryAvailablePurchasesResults(Result, ProvidedProductInformation);
-				}
-				PRAGMA_ENABLE_DEPRECATION_WARNINGS
+				}				
 			}
 			FPlatformMisc::LowLevelOutputDebugStringf(TEXT("In-App Purchase query was completed  %s\n"), Result == EInAppPurchaseState::Success ? TEXT("successfully") : TEXT("unsuccessfully"));
 		}),
@@ -203,12 +201,10 @@ JNI_METHOD void Java_com_epicgames_ue4_AmazonStoreHelper_nativePurchaseComplete(
 			if (IOnlineSubsystem* const OnlineSub = IOnlineSubsystem::Get())
 			{
 				// call store implementation to process query results.
-				PRAGMA_DISABLE_DEPRECATION_WARNINGS
 				if (FOnlineStoreGameCircle* StoreInterface = (FOnlineStoreGameCircle*)OnlineSub->GetStoreInterface().Get())
 				{
 					StoreInterface->ProcessPurchaseResult(Result, ProductId, ReceiptData, Signature);
 				}
-				PRAGMA_ENABLE_DEPRECATION_WARNINGS
 			}
 		}),
 		GET_STATID(STAT_FSimpleDelegateGraphTask_ProcessIapResult), 
@@ -304,7 +300,6 @@ JNI_METHOD void Java_com_epicgames_ue4_AmazonStoreHelper_nativeRestorePurchasesC
 		{
 			FPlatformMisc::LowLevelOutputDebugStringf(TEXT("Sending result back to OnlineSubsystem.\n"));
 			// call store implementation to process query results.
-			PRAGMA_DISABLE_DEPRECATION_WARNINGS
 			if (FOnlineStoreGameCircle* StoreInterface = (FOnlineStoreGameCircle*)OnlineSub->GetStoreInterface().Get())
 			{
 				if (StoreInterface->CachedPurchaseRestoreObject.IsValid())
@@ -314,7 +309,6 @@ JNI_METHOD void Java_com_epicgames_ue4_AmazonStoreHelper_nativeRestorePurchasesC
 				}
 				StoreInterface->TriggerOnInAppPurchaseRestoreCompleteDelegates(bSuccess ? EInAppPurchaseState::Restored : Result);
 			}
-			PRAGMA_ENABLE_DEPRECATION_WARNINGS
 		}
 		}),
 		GET_STATID(STAT_FSimpleDelegateGraphTask_RestorePurchases),
diff --git a/Engine/Plugins/Online/Android/OnlineSubsystemGameCircle/Source/OnlineSubsystemGameCircle/Private/OnlineSubsystemGameCircle.cpp b/Engine/Plugins/Online/Android/OnlineSubsystemGameCircle/Source/OnlineSubsystemGameCircle/Private/OnlineSubsystemGameCircle.cpp
--- a/Engine/Plugins/Online/Android/OnlineSubsystemGameCircle/Source/OnlineSubsystemGameCircle/Private/OnlineSubsystemGameCircle.cpp
+++ b/Engine/Plugins/Online/Android/OnlineSubsystemGameCircle/Source/OnlineSubsystemGameCircle/Private/OnlineSubsystemGameCircle.cpp
@@ -19,12 +19,10 @@ IOnlineIdentityPtr FOnlineSubsystemGameCircle::GetIdentityInterface() const
 	return IdentityInterface;
 }
 
-PRAGMA_DISABLE_DEPRECATION_WARNINGS
 IOnlineStorePtr FOnlineSubsystemGameCircle::GetStoreInterface() const
 {
 	return StoreInterface;
 }
-PRAGMA_ENABLE_DEPRECATION_WARNINGS
 
 IOnlineSessionPtr FOnlineSubsystemGameCircle::GetSessionInterface() const
 {
diff --git a/Engine/Plugins/Online/Android/OnlineSubsystemGooglePlay/Source/Private/OnlineSubsystemGooglePlay.cpp b/Engine/Plugins/Online/Android/OnlineSubsystemGooglePlay/Source/Private/OnlineSubsystemGooglePlay.cpp
--- a/Engine/Plugins/Online/Android/OnlineSubsystemGooglePlay/Source/Private/OnlineSubsystemGooglePlay.cpp
+++ b/Engine/Plugins/Online/Android/OnlineSubsystemGooglePlay/Source/Private/OnlineSubsystemGooglePlay.cpp
@@ -47,12 +47,10 @@ IOnlineIdentityPtr FOnlineSubsystemGooglePlay::GetIdentityInterface() const
 	return IdentityInterface;
 }
 
-PRAGMA_DISABLE_DEPRECATION_WARNINGS
 IOnlineStorePtr FOnlineSubsystemGooglePlay::GetStoreInterface() const
 {
 	return StoreInterface;
 }
-PRAGMA_ENABLE_DEPRECATION_WARNINGS
 
 IOnlineStoreV2Ptr FOnlineSubsystemGooglePlay::GetStoreV2Interface() const
 {
diff --git a/Engine/Plugins/Online/Android/OnlineSubsystemGooglePlay/Source/Public/OnlineSubsystemGooglePlay.h b/Engine/Plugins/Online/Android/OnlineSubsystemGooglePlay/Source/Public/OnlineSubsystemGooglePlay.h
--- a/Engine/Plugins/Online/Android/OnlineSubsystemGooglePlay/Source/Public/OnlineSubsystemGooglePlay.h
+++ b/Engine/Plugins/Online/Android/OnlineSubsystemGooglePlay/Source/Public/OnlineSubsystemGooglePlay.h
@@ -58,9 +58,7 @@ public:
 	virtual IOnlineIdentityPtr GetIdentityInterface() const override;
 	virtual IOnlineTitleFilePtr GetTitleFileInterface() const override;
 	virtual IOnlineEntitlementsPtr GetEntitlementsInterface() const override;
-	PRAGMA_DISABLE_DEPRECATION_WARNINGS
 	virtual IOnlineStorePtr GetStoreInterface() const override;
-	PRAGMA_ENABLE_DEPRECATION_WARNINGS
 	virtual IOnlineStoreV2Ptr GetStoreV2Interface() const override;
 	virtual IOnlinePurchasePtr GetPurchaseInterface() const override;
 	virtual IOnlineEventsPtr GetEventsInterface() const override { return nullptr; }
diff --git a/Engine/Plugins/Online/IOS/OnlineSubsystemIOS/Source/Private/OnlineSubsystemIOS.cpp b/Engine/Plugins/Online/IOS/OnlineSubsystemIOS/Source/Private/OnlineSubsystemIOS.cpp
--- a/Engine/Plugins/Online/IOS/OnlineSubsystemIOS/Source/Private/OnlineSubsystemIOS.cpp
+++ b/Engine/Plugins/Online/IOS/OnlineSubsystemIOS/Source/Private/OnlineSubsystemIOS.cpp
@@ -79,12 +79,10 @@ IOnlineEntitlementsPtr FOnlineSubsystemIOS::GetEntitlementsInterface() const
 	return nullptr;
 }
 
-PRAGMA_DISABLE_DEPRECATION_WARNINGS
 IOnlineStorePtr FOnlineSubsystemIOS::GetStoreInterface() const
 {
 	return StoreInterface;
 }
-PRAGMA_ENABLE_DEPRECATION_WARNINGS
 
 IOnlineStoreV2Ptr FOnlineSubsystemIOS::GetStoreV2Interface() const
 {
diff --git a/Engine/Plugins/Online/IOS/OnlineSubsystemIOS/Source/Public/OnlineSubsystemIOS.h b/Engine/Plugins/Online/IOS/OnlineSubsystemIOS/Source/Public/OnlineSubsystemIOS.h
--- a/Engine/Plugins/Online/IOS/OnlineSubsystemIOS/Source/Public/OnlineSubsystemIOS.h
+++ b/Engine/Plugins/Online/IOS/OnlineSubsystemIOS/Source/Public/OnlineSubsystemIOS.h
@@ -47,9 +47,7 @@ public:
 	virtual IOnlineIdentityPtr GetIdentityInterface() const override;
 	virtual IOnlineTitleFilePtr GetTitleFileInterface() const override;
 	virtual IOnlineEntitlementsPtr GetEntitlementsInterface() const override;
-	PRAGMA_DISABLE_DEPRECATION_WARNINGS
 	virtual IOnlineStorePtr GetStoreInterface() const override;
-	PRAGMA_ENABLE_DEPRECATION_WARNINGS
 	virtual IOnlineStoreV2Ptr GetStoreV2Interface() const override;
 	virtual IOnlinePurchasePtr GetPurchaseInterface() const override;
 	virtual IOnlineEventsPtr GetEventsInterface() const override;
diff --git a/Engine/Plugins/Online/OnlineSubsystem/Source/Public/OnlineSubsystem.h b/Engine/Plugins/Online/OnlineSubsystem/Source/Public/OnlineSubsystem.h
--- a/Engine/Plugins/Online/OnlineSubsystem/Source/Public/OnlineSubsystem.h
+++ b/Engine/Plugins/Online/OnlineSubsystem/Source/Public/OnlineSubsystem.h
@@ -373,7 +373,8 @@ public:
 	 * Get the interface for accessing an online store
 	 * @return Interface pointer for the appropriate online store service
 	 */
-	UE_DEPRECATED(4.25, "Move to OnlineStoreInterfaceV2 and OnlinePurchaseInterface.")
+	// @todo samz
+	//	UE_DEPRECATED(4.26, "Move to OnlineStoreInterfaceV2 and OnlinePurchaseInterface.")
 	virtual IOnlineStorePtr GetStoreInterface() const { return nullptr; }
 
 	/** 
diff --git a/Engine/Plugins/Online/OnlineSubsystemUtils/Source/OnlineSubsystemUtils/Private/InAppPurchaseCallbackProxy.cpp b/Engine/Plugins/Online/OnlineSubsystemUtils/Source/OnlineSubsystemUtils/Private/InAppPurchaseCallbackProxy.cpp
--- a/Engine/Plugins/Online/OnlineSubsystemUtils/Source/OnlineSubsystemUtils/Private/InAppPurchaseCallbackProxy.cpp
+++ b/Engine/Plugins/Online/OnlineSubsystemUtils/Source/OnlineSubsystemUtils/Private/InAppPurchaseCallbackProxy.cpp
@@ -28,9 +28,7 @@ void UInAppPurchaseCallbackProxy::Trigger(APlayerController* PlayerController, c
 	{
 		if (IOnlineSubsystem* const OnlineSub = IOnlineSubsystem::IsLoaded() ? IOnlineSubsystem::Get() : nullptr)
 		{
-			PRAGMA_DISABLE_DEPRECATION_WARNINGS
 			IOnlineStorePtr StoreInterface = OnlineSub->GetStoreInterface();
-			PRAGMA_ENABLE_DEPRECATION_WARNINGS
 			if (StoreInterface.IsValid())
 			{
 				bFailedToEvenSubmit = false;
@@ -125,9 +123,7 @@ void UInAppPurchaseCallbackProxy::RemoveDelegate()
 	{
 		if (IOnlineSubsystem* OnlineSub = IOnlineSubsystem::IsLoaded() ? IOnlineSubsystem::Get() : nullptr)
 		{
-			PRAGMA_DISABLE_DEPRECATION_WARNINGS
 			IOnlineStorePtr InAppPurchases = OnlineSub->GetStoreInterface();
-			PRAGMA_ENABLE_DEPRECATION_WARNINGS
 			if (InAppPurchases.IsValid())
 			{
 				InAppPurchases->ClearOnInAppPurchaseCompleteDelegate_Handle(InAppPurchaseCompleteDelegateHandle);
diff --git a/Engine/Plugins/Online/OnlineSubsystemUtils/Source/OnlineSubsystemUtils/Private/InAppPurchaseQueryCallbackProxy.cpp b/Engine/Plugins/Online/OnlineSubsystemUtils/Source/OnlineSubsystemUtils/Private/InAppPurchaseQueryCallbackProxy.cpp
--- a/Engine/Plugins/Online/OnlineSubsystemUtils/Source/OnlineSubsystemUtils/Private/InAppPurchaseQueryCallbackProxy.cpp
+++ b/Engine/Plugins/Online/OnlineSubsystemUtils/Source/OnlineSubsystemUtils/Private/InAppPurchaseQueryCallbackProxy.cpp
@@ -23,9 +23,7 @@ void UInAppPurchaseQueryCallbackProxy::TriggerQuery(APlayerController* PlayerCon
 	{
 		if (IOnlineSubsystem* const OnlineSub = IOnlineSubsystem::IsLoaded() ? IOnlineSubsystem::Get() : nullptr)
 		{
-			PRAGMA_DISABLE_DEPRECATION_WARNINGS
 			IOnlineStorePtr StoreInterface = OnlineSub->GetStoreInterface();
-			PRAGMA_ENABLE_DEPRECATION_WARNINGS
 			if (StoreInterface.IsValid())
 			{
 				bFailedToEvenSubmit = false;
@@ -111,9 +109,7 @@ void UInAppPurchaseQueryCallbackProxy::RemoveDelegate()
 	{
 		if (IOnlineSubsystem* OnlineSub = IOnlineSubsystem::IsLoaded() ? IOnlineSubsystem::Get() : nullptr)
 		{
-			PRAGMA_DISABLE_DEPRECATION_WARNINGS
 			IOnlineStorePtr InAppPurchases = OnlineSub->GetStoreInterface();
-			PRAGMA_ENABLE_DEPRECATION_WARNINGS
 			if (InAppPurchases.IsValid())
 			{
 				InAppPurchases->ClearOnQueryForAvailablePurchasesCompleteDelegate_Handle(InAppPurchaseReadCompleteDelegateHandle);
diff --git a/Engine/Plugins/Online/OnlineSubsystemUtils/Source/OnlineSubsystemUtils/Private/InAppPurchaseRestoreCallbackProxy.cpp b/Engine/Plugins/Online/OnlineSubsystemUtils/Source/OnlineSubsystemUtils/Private/InAppPurchaseRestoreCallbackProxy.cpp
--- a/Engine/Plugins/Online/OnlineSubsystemUtils/Source/OnlineSubsystemUtils/Private/InAppPurchaseRestoreCallbackProxy.cpp
+++ b/Engine/Plugins/Online/OnlineSubsystemUtils/Source/OnlineSubsystemUtils/Private/InAppPurchaseRestoreCallbackProxy.cpp
@@ -25,9 +25,7 @@ void UInAppPurchaseRestoreCallbackProxy::Trigger(const TArray<FInAppPurchaseProd
 	{
 		if (IOnlineSubsystem* const OnlineSub = IOnlineSubsystem::IsLoaded() ? IOnlineSubsystem::Get() : nullptr)
 		{
-			PRAGMA_DISABLE_DEPRECATION_WARNINGS
 			IOnlineStorePtr StoreInterface = OnlineSub->GetStoreInterface();
-			PRAGMA_ENABLE_DEPRECATION_WARNINGS
 			if (StoreInterface.IsValid())
 			{
 				bFailedToEvenSubmit = false;
@@ -111,9 +109,7 @@ void UInAppPurchaseRestoreCallbackProxy::RemoveDelegate()
 	{
 		if (IOnlineSubsystem* OnlineSub = IOnlineSubsystem::IsLoaded() ? IOnlineSubsystem::Get() : nullptr)
 		{
-			PRAGMA_DISABLE_DEPRECATION_WARNINGS
 			IOnlineStorePtr InAppPurchases = OnlineSub->GetStoreInterface();
-			PRAGMA_ENABLE_DEPRECATION_WARNINGS
 			if (InAppPurchases.IsValid())
 			{
 				InAppPurchases->ClearOnInAppPurchaseRestoreCompleteDelegate_Handle(InAppPurchaseRestoreCompleteDelegateHandle);
